Disable RayTracing in Vulkan until shader issue is resolved.

diff --git a/Engine/Config/VulkanPC/DataDrivenPlatformInfo.ini b/Engine/Config/VulkanPC/DataDrivenPlatformInfo.ini
--- a/Engine/Config/VulkanPC/DataDrivenPlatformInfo.ini
+++ b/Engine/Config/VulkanPC/DataDrivenPlatformInfo.ini
@@ -49,9 +49,9 @@ bSupports4ComponentUAVReadWrite = false
 ; //#todo-vulkanddpi
 bSupportsRenderTargetWriteMask = false
 
-bSupportsRayTracing = true
+bSupportsRayTracing = false
 bSupportsRayTracingShaders = false
-bSupportsInlineRayTracing = true
+bSupportsInlineRayTracing = false
 bSupportsRayTracingIndirectInstanceData = false
 bSupportsGPUSkinCache = true
 bSupportsByteBufferComputeShaders = true
