Re-enable RayTracing in Vulkan

diff --git a/Engine/Config/VulkanPC/DataDrivenPlatformInfo.ini b/Engine/Config/VulkanPC/DataDrivenPlatformInfo.ini
--- a/Engine/Config/VulkanPC/DataDrivenPlatformInfo.ini
+++ b/Engine/Config/VulkanPC/DataDrivenPlatformInfo.ini
@@ -49,9 +49,9 @@ bSupports4ComponentUAVReadWrite = false
 ; //#todo-vulkanddpi
 bSupportsRenderTargetWriteMask = false
 
-bSupportsRayTracing = false
+bSupportsRayTracing = true
 bSupportsRayTracingShaders = false
-bSupportsInlineRayTracing = false
+bSupportsInlineRayTracing = true
 bSupportsRayTracingIndirectInstanceData = false
 bSupportsGPUSkinCache = true
 bSupportsByteBufferComputeShaders = true
