VulkanRHI: Fix Anisotropy EngineTests.

diff --git a/Engine/Config/VulkanPC/DataDrivenPlatformInfo.ini b/Engine/Config/VulkanPC/DataDrivenPlatformInfo.ini
--- a/Engine/Config/VulkanPC/DataDrivenPlatformInfo.ini
+++ b/Engine/Config/VulkanPC/DataDrivenPlatformInfo.ini
@@ -127,6 +127,8 @@ bSupportsDiaphragmDOF=true
 ; shaderOutputLayer is optional in our baseline VK version
 bSupportsVertexShaderLayer=false
 
+bSupportsAnisotropicMaterials=true
+
 
 [ShaderPlatform VULKAN_PCES3_1]
 Language=Vulkan
