Until Vulcan supports SM6 properly Nanite is disabled.

diff --git a/Engine/Config/VulkanPC/DataDrivenPlatformInfo.ini b/Engine/Config/VulkanPC/DataDrivenPlatformInfo.ini
--- a/Engine/Config/VulkanPC/DataDrivenPlatformInfo.ini
+++ b/Engine/Config/VulkanPC/DataDrivenPlatformInfo.ini
@@ -63,7 +63,7 @@ bSupportsUInt64ImageAtomics = true
 ; //#todo-vulkanddpi
 bRequiresVendorExtensionsForAtomics = false
 
-bSupportsNanite = true
+bSupportsNanite = false
 bSupportsSceneDataCompressedTransforms = true
 bSupportsLumenGI = true
 
