Enable ComputeFramework on PC Vulkan.

diff --git a/Engine/Config/VulkanPC/DataDrivenPlatformInfo.ini b/Engine/Config/VulkanPC/DataDrivenPlatformInfo.ini
--- a/Engine/Config/VulkanPC/DataDrivenPlatformInfo.ini
+++ b/Engine/Config/VulkanPC/DataDrivenPlatformInfo.ini
@@ -92,8 +92,7 @@ bRequiresExplicit128bitRT = false
 bTargetsTiledGPU = false
 bNeedsOfflineCompiler = false
 
-; //#todo-vulkanddpi
-bSupportsComputeFramework = false
+bSupportsComputeFramework = true
 
 ; //#todo-vulkanddpi
 bSupportsDualSourceBlending = true
