VulkanRHI: Add VULKAN_SM6.  Move features and requirements around to match what was done on d3d.

diff --git a/Engine/Config/VulkanPC/DataDrivenPlatformInfo.ini b/Engine/Config/VulkanPC/DataDrivenPlatformInfo.ini
--- a/Engine/Config/VulkanPC/DataDrivenPlatformInfo.ini
+++ b/Engine/Config/VulkanPC/DataDrivenPlatformInfo.ini
@@ -1,6 +1,6 @@
 ; Have a single source of truth for all desktop class Vulkan platforms
 
-; This is done so all the platforms using SP_VULKAN_SM5 can share the same DDC
+; This is done so all the platforms using SP_VULKAN_SM5/SM6 can share the same DDC
 ; If your Vulkan platform needs its own settings, create a new SP for it
 
 [DataDrivenPlatformInfo]
@@ -19,6 +19,7 @@ bIsPC = true
 bIsConsole = false
 bIsAndroidOpenGLES = false
 bSupportsDxc = true
+bIsSPIRV=true
 
 bSupportsDebugViewShaders=true
 
@@ -57,14 +58,14 @@ bSupportsGPUScene = true
 ; //#todo-vulkanddpi
 bSupportsPrimitiveShaders = false
 
-bSupportsUInt64ImageAtomics = true
+bSupportsUInt64ImageAtomics = false
 
 ; //#todo-vulkanddpi
 bRequiresVendorExtensionsForAtomics = false
 
-bSupportsNanite = true
+bSupportsNanite = false
 bSupportsSceneDataCompressedTransforms = true
-bSupportsLumenGI = true
+bSupportsLumenGI = false
 
 ; //#todo-vulkanddpi
 bSupportsSSDIndirect = false
@@ -110,7 +111,6 @@ bSupportsGen5TemporalAA=true
 ; //#todo-vulkanddpi
 bSupportsPerPixelDBufferMask = false
 
-bIsSPIRV=true
 
 bSupportsVariableRateShading = false
 
@@ -132,6 +132,121 @@ bSupportsOIT=false
 EnablesHLSL2021ByDefault=1
 FriendlyName=LOCTEXT("FriendlyShaderPlatformName_VulkanPC_SM5", "PC Vulkan SM5")
 
+BindlessSupport=Unsupported
+
+
+
+[ShaderPlatform VULKAN_SM6]
+Language=Vulkan
+MaxFeatureLevel=SM6
+ShaderFormat=SF_VULKAN_SM6
+
+bIsMobile = false
+bIsMetalMRT = false
+bIsPC = true
+bIsConsole = false
+bIsAndroidOpenGLES = false
+bSupportsDxc = true
+bIsSPIRV=true
+
+bSupportsDebugViewShaders=true
+bSupportsMobileMultiView = false
+bSupportsArrayTextureCompression = false
+
+bSupportsDistanceFields = true
+bSupportsDiaphragmDOF = true
+bSupportsRGBColorBuffer = true
+bSupportsCapsuleShadows = true
+bSupportsPercentageCloserShadows=true
+bSupportsVolumetricFog = true
+bSupportsIndexBufferUAVs = true
+
+bSupportsInstancedStereo = true
+SupportsMultiViewport = RuntimeDependent
+
+bSupportsMSAA = true
+bSupports4ComponentUAVReadWrite = false
+bSupportsRenderTargetWriteMask = false
+
+bSupportsRayTracing = true
+bSupportsRayTracingShaders = false
+bSupportsInlineRayTracing = true
+bSupportsRayTracingIndirectInstanceData = true
+bSupportsByteBufferComputeShaders = true
+
+bSupportsGPUScene = true
+bSupportsPrimitiveShaders = false
+bSupportsUInt64ImageAtomics = true
+bRequiresVendorExtensionsForAtomics = false
+
+bSupportsNanite = true
+bSupportsSceneDataCompressedTransforms = true
+bSupportsLumenGI = true
+
+; //#todo-vulkanddpi
+bSupportsSSDIndirect = false
+
+; //#todo-vulkanddpi
+bSupportsTemporalHistoryUpscale = false
+
+; //#todo-vulkanddpi
+bSupportsRTIndexFromVS = false
+
+; //#todo-vulkanddpi
+bSupportsIntrinsicWaveOnce = false
+
+; //#todo-vulkanddpi
+bSupportsConservativeRasterization = false
+
+bSupportsWaveOperations=RuntimeGuaranteed
+MinimumWaveSize=4
+MaximumWaveSize=128
+
+; //#todo-vulkanddpi
+bRequiresExplicit128bitRT = false
+
+bTargetsTiledGPU = false
+bNeedsOfflineCompiler = false
+
+bSupportsComputeFramework = true
+
+; //#todo-vulkanddpi
+bSupportsDualSourceBlending = true
+
+; //#todo-vulkanddpi
+bRequiresGeneratePrevTransformBuffer = false
+
+bRequiresRenderTargetDuringRaster = true
+
+; //#todo-vulkanddpi
+bRequiresDisableForwardLocalLights = true
+
+bCompileSignalProcessingPipeline = true
+bSupportsGen5TemporalAA=true
+
+; //#todo-vulkanddpi
+bSupportsPerPixelDBufferMask = false
+
+bSupportsVariableRateShading = false
+
+NumberOfComputeThreads = 64
+
+bSupportsHairStrandGeometry=true
+bSupportsFFTBloom=true
+bSupportsDiaphragmDOF=true
+
+; shaderOutputLayer is optional in our baseline VK version
+bSupportsVertexShaderLayer=false
+
+bSupportsAnisotropicMaterials=true
+
+; disable until DXC -> SPRIV ROV support is implemented
+bSupportsROV=false
+bSupportsOIT=false
+
+EnablesHLSL2021ByDefault=1
+FriendlyName=LOCTEXT("FriendlyShaderPlatformName_VulkanPC_SM6", "PC Vulkan SM6")
+
 BindlessSupport=AllShaderTypes
 
 
diff --git a/Engine/Config/Windows/BaseWindowsEngine.ini b/Engine/Config/Windows/BaseWindowsEngine.ini
--- a/Engine/Config/Windows/BaseWindowsEngine.ini
+++ b/Engine/Config/Windows/BaseWindowsEngine.ini
@@ -21,3 +21,7 @@ BindlessSamplers=RayTracingOnly
 [SF_VULKAN_SM5]
 BindlessResources=Disabled
 BindlessSamplers=Disabled
+
+[SF_VULKAN_SM6]
+BindlessResources=Disabled
+BindlessSamplers=Disabled
diff --git a/Engine/Plugins/FX/Niagara/Source/NiagaraEditor/Private/ViewModels/NiagaraScriptStatsViewModel.cpp b/Engine/Plugins/FX/Niagara/Source/NiagaraEditor/Private/ViewModels/NiagaraScriptStatsViewModel.cpp
--- a/Engine/Plugins/FX/Niagara/Source/NiagaraEditor/Private/ViewModels/NiagaraScriptStatsViewModel.cpp
+++ b/Engine/Plugins/FX/Niagara/Source/NiagaraEditor/Private/ViewModels/NiagaraScriptStatsViewModel.cpp
@@ -768,6 +768,7 @@ void FNiagaraScriptStatsViewModel::BuildShaderPlatformDetails()
 	TArray<EShaderPlatform, TInlineAllocator<SP_StaticPlatform_Last + 1>> StaticShaderPlatforms;
 	StaticShaderPlatforms.Add(SP_PCD3D_SM5);
 	StaticShaderPlatforms.Add(SP_VULKAN_SM5);
+	StaticShaderPlatforms.Add(SP_VULKAN_SM6);
 	StaticShaderPlatforms.Add(SP_OPENGL_ES3_1_ANDROID); 
 	StaticShaderPlatforms.Add(SP_VULKAN_ES3_1_ANDROID);
 	StaticShaderPlatforms.Add(SP_VULKAN_SM5_ANDROID);
diff --git a/Engine/Source/Developer/Linux/LinuxTargetPlatform/Private/LinuxTargetPlatform.h b/Engine/Source/Developer/Linux/LinuxTargetPlatform/Private/LinuxTargetPlatform.h
--- a/Engine/Source/Developer/Linux/LinuxTargetPlatform/Private/LinuxTargetPlatform.h
+++ b/Engine/Source/Developer/Linux/LinuxTargetPlatform/Private/LinuxTargetPlatform.h
@@ -263,9 +263,11 @@ public:
 		if (!TProperties::IsServerOnly())
 		{
 			static FName NAME_VULKAN_SM5(TEXT("SF_VULKAN_SM5"));
+			static FName NAME_VULKAN_SM6(TEXT("SF_VULKAN_SM6"));
 			static FName NAME_VULKAN_ES31(TEXT("SF_VULKAN_ES31"));
 
 			OutFormats.AddUnique(NAME_VULKAN_SM5);
+			OutFormats.AddUnique(NAME_VULKAN_SM6);
 			OutFormats.AddUnique(NAME_VULKAN_ES31);
 		}
 	}
diff --git a/Engine/Source/Developer/VulkanShaderFormat/Private/VulkanShaderCompiler.cpp b/Engine/Source/Developer/VulkanShaderFormat/Private/VulkanShaderCompiler.cpp
--- a/Engine/Source/Developer/VulkanShaderFormat/Private/VulkanShaderCompiler.cpp
+++ b/Engine/Source/Developer/VulkanShaderFormat/Private/VulkanShaderCompiler.cpp
@@ -49,7 +49,7 @@ inline bool IsVulkanShaderFormat(FName ShaderFormat)
 	return ShaderFormat == NAME_VULKAN_ES3_1_ANDROID
 		|| ShaderFormat == NAME_VULKAN_ES3_1
 		|| ShaderFormat == NAME_VULKAN_SM5
-//		|| ShaderFormat == NAME_VULKAN_SM6 // :todo-jn:
+		|| ShaderFormat == NAME_VULKAN_SM6
 		|| ShaderFormat == NAME_VULKAN_SM5_ANDROID;
 }
 
@@ -2365,10 +2365,13 @@ static bool CompileWithShaderConductor(
 		Options.HlslVersion = 2021;
 	}
 
-	// Ray tracing features require Vulkan 1.2 environment.
+	// Ray tracing features require Vulkan 1.2 environment minimum.
 	if (bIsRayTracingShader || Input.Environment.CompilerFlags.Contains(CFLAG_InlineRayTracing))
 	{
-		Options.TargetEnvironment = CrossCompiler::FShaderConductorOptions::ETargetEnvironment::Vulkan_1_2;
+		if (Options.TargetEnvironment < CrossCompiler::FShaderConductorOptions::ETargetEnvironment::Vulkan_1_2)
+		{
+			Options.TargetEnvironment = CrossCompiler::FShaderConductorOptions::ETargetEnvironment::Vulkan_1_2;
+		}
 	}
 	else
 	{
diff --git a/Engine/Source/Developer/VulkanShaderFormat/Private/VulkanShaderFormat.cpp b/Engine/Source/Developer/VulkanShaderFormat/Private/VulkanShaderFormat.cpp
--- a/Engine/Source/Developer/VulkanShaderFormat/Private/VulkanShaderFormat.cpp
+++ b/Engine/Source/Developer/VulkanShaderFormat/Private/VulkanShaderFormat.cpp
@@ -19,11 +19,16 @@ class FShaderFormatVulkan : public IShaderFormat
 	{
 		UE_SHADER_VULKAN_ES3_1_VER	= 37,
 		UE_SHADER_VULKAN_SM5_VER 	= 37,
+		UE_SHADER_VULKAN_SM6_VER	= 37,
 	};
 
 	int32 InternalGetVersion(FName Format) const
 	{
-		if (Format == NAME_VULKAN_SM5 || Format == NAME_VULKAN_SM5_ANDROID)
+		if (Format == NAME_VULKAN_SM6)
+		{
+			return UE_SHADER_VULKAN_SM6_VER;
+		}
+		else if (Format == NAME_VULKAN_SM5 || Format == NAME_VULKAN_SM5_ANDROID)
 		{
 			return UE_SHADER_VULKAN_SM5_VER;
 		}
@@ -73,6 +78,7 @@ public:
 		OutFormats.Add(NAME_VULKAN_ES3_1_ANDROID);
 		OutFormats.Add(NAME_VULKAN_ES3_1);
 		OutFormats.Add(NAME_VULKAN_SM5_ANDROID);
+		OutFormats.Add(NAME_VULKAN_SM6);
 	}
 
 	virtual void CompileShader(FName Format, const struct FShaderCompilerInput& Input, struct FShaderCompilerOutput& Output,const FString& WorkingDirectory) const
@@ -94,6 +100,10 @@ public:
 		{
 			DoCompileVulkanShader(Input, Output, WorkingDirectory, EVulkanShaderVersion::SM5);
 		}
+		else if (Format == NAME_VULKAN_SM6)
+		{
+			DoCompileVulkanShader(Input, Output, WorkingDirectory, EVulkanShaderVersion::SM6);
+		}
 	}
 
 	virtual const TCHAR* GetPlatformIncludeDirectory() const
diff --git a/Engine/Source/Developer/Windows/WindowsTargetPlatform/Private/WindowsTargetPlatformClasses.cpp b/Engine/Source/Developer/Windows/WindowsTargetPlatform/Private/WindowsTargetPlatformClasses.cpp
--- a/Engine/Source/Developer/Windows/WindowsTargetPlatform/Private/WindowsTargetPlatformClasses.cpp
+++ b/Engine/Source/Developer/Windows/WindowsTargetPlatform/Private/WindowsTargetPlatformClasses.cpp
@@ -34,7 +34,7 @@ static bool FilterShaderPlatform_D3D11(const FString& InShaderPlatform)
 
 static bool FilterShaderPlatform_Vulkan(const FString& InShaderPlatform)
 {
-	return InShaderPlatform == TEXT("SF_VULKAN_SM5");
+	return InShaderPlatform == TEXT("SF_VULKAN_SM5") || InShaderPlatform == TEXT("SF_VULKAN_SM6");
 }
 
 template<typename TFilterFunc>
diff --git a/Engine/Source/Developer/Windows/WindowsTargetPlatform/Public/GenericWindowsTargetPlatform.h b/Engine/Source/Developer/Windows/WindowsTargetPlatform/Public/GenericWindowsTargetPlatform.h
--- a/Engine/Source/Developer/Windows/WindowsTargetPlatform/Public/GenericWindowsTargetPlatform.h
+++ b/Engine/Source/Developer/Windows/WindowsTargetPlatform/Public/GenericWindowsTargetPlatform.h
@@ -69,6 +69,7 @@ public:
 		static FName NAME_PCD3D_SM6(TEXT("PCD3D_SM6"));
 		static FName NAME_PCD3D_SM5(TEXT("PCD3D_SM5"));
 		static FName NAME_VULKAN_SM5(TEXT("SF_VULKAN_SM5"));
+		static FName NAME_VULKAN_SM6(TEXT("SF_VULKAN_SM6"));
 		static FName NAME_PCD3D_ES3_1(TEXT("PCD3D_ES31"));
 		static FName NAME_SF_VULKAN_ES31(TEXT("SF_VULKAN_ES31"));
 		static FName NAME_OPENGL_150_ES3_1(TEXT("GLSL_150_ES31"));
@@ -208,6 +209,7 @@ public:
 			static FName NAME_VULKAN_ES31(TEXT("SF_VULKAN_ES31"));
 			static FName NAME_OPENGL_150_ES3_1(TEXT("GLSL_150_ES31"));
 			static FName NAME_VULKAN_SM5(TEXT("SF_VULKAN_SM5"));
+			static FName NAME_VULKAN_SM6(TEXT("SF_VULKAN_SM6"));
 			static FName NAME_PCD3D_ES3_1(TEXT("PCD3D_ES31"));
 
 			OutFormats.AddUnique(NAME_PCD3D_SM5);
@@ -215,6 +217,7 @@ public:
 			OutFormats.AddUnique(NAME_VULKAN_ES31);
 			OutFormats.AddUnique(NAME_OPENGL_150_ES3_1);
 			OutFormats.AddUnique(NAME_VULKAN_SM5);
+			OutFormats.AddUnique(NAME_VULKAN_SM6);
 			OutFormats.AddUnique(NAME_PCD3D_ES3_1);
 		}
 	}
diff --git a/Engine/Source/Editor/DetailCustomizations/Private/LinuxTargetSettingsDetails.cpp b/Engine/Source/Editor/DetailCustomizations/Private/LinuxTargetSettingsDetails.cpp
--- a/Engine/Source/Editor/DetailCustomizations/Private/LinuxTargetSettingsDetails.cpp
+++ b/Engine/Source/Editor/DetailCustomizations/Private/LinuxTargetSettingsDetails.cpp
@@ -66,6 +66,10 @@ static FText GetFriendlyNameFromLinuxShaderFormat(const FName InShaderFormat)
 	{
 		FriendlyRHIName = LOCTEXT("VulkanSM5", "Vulkan Desktop (SM5)");
 	}
+	else if (InShaderFormat == NAME_VULKAN_SM6)
+	{
+		FriendlyRHIName = LOCTEXT("VulkanSM6", "Vulkan Desktop (SM6)");
+	}
 	else if (InShaderFormat == TEXT("GLSL_430"))
 	{
 		// Explicitly remove these formats as they are obsolete/not quite supported; users can still target them by adding them as +TargetedRHIs in the TargetPlatform ini.
diff --git a/Engine/Source/Editor/DetailCustomizations/Private/MacTargetSettingsDetails.cpp b/Engine/Source/Editor/DetailCustomizations/Private/MacTargetSettingsDetails.cpp
--- a/Engine/Source/Editor/DetailCustomizations/Private/MacTargetSettingsDetails.cpp
+++ b/Engine/Source/Editor/DetailCustomizations/Private/MacTargetSettingsDetails.cpp
@@ -147,6 +147,9 @@ static FText GetFriendlyNameFromRHINameMac(FName InRHIName)
 	case SP_VULKAN_SM5_ANDROID:
 		FriendlyRHIName = LOCTEXT("VulkanSM5", "Vulkan (SM5)");
 		break;
+	case SP_VULKAN_SM6:
+		FriendlyRHIName = LOCTEXT("VulkanSM6", "Vulkan (SM6)");
+		break;
 	case SP_VULKAN_PCES3_1:
 	case SP_VULKAN_ES3_1_ANDROID:
 		FriendlyRHIName = LOCTEXT("VulkanES31", "Vulkan (ES 3.1)");
diff --git a/Engine/Source/Editor/DetailCustomizations/Private/WindowsTargetSettingsDetails.cpp b/Engine/Source/Editor/DetailCustomizations/Private/WindowsTargetSettingsDetails.cpp
--- a/Engine/Source/Editor/DetailCustomizations/Private/WindowsTargetSettingsDetails.cpp
+++ b/Engine/Source/Editor/DetailCustomizations/Private/WindowsTargetSettingsDetails.cpp
@@ -73,6 +73,9 @@ static FText GetFriendlyNameFromWindowsShaderPlatform(FName InShaderPlatformName
 	case SP_VULKAN_SM5:
 		FriendlyRHIName = LOCTEXT("VulkanSM5", "Vulkan (SM5)");
 		break;
+	case SP_VULKAN_SM6:
+		FriendlyRHIName = LOCTEXT("VulkanSM6", "Vulkan (SM6)");
+		break;
 
 	case SP_OPENGL_PCES3_1:
 	case SP_VULKAN_PCES3_1:
@@ -98,17 +101,16 @@ static FText GetFriendlyNameForWindowsShaderPlatformCheckbox(FName InShaderPlatf
 	switch (ShaderPlatform)
 	{
 	case SP_PCD3D_SM6:
+	case SP_VULKAN_SM6:
 		FriendlyName = LOCTEXT("SM6", "SM6");
 		break;
 	case SP_PCD3D_SM5:
+	case SP_VULKAN_SM5:
 		FriendlyName = LOCTEXT("SM5", "SM5");
 		break;
 	case SP_PCD3D_ES3_1:
 		FriendlyName = LOCTEXT("ES31", "ES3.1");
 		break;
-	case SP_VULKAN_SM5:
-		FriendlyName = LOCTEXT("SM5", "SM5");
-		break;
 	default:
 		break;
 	}
@@ -128,7 +130,7 @@ static bool FilterShaderPlatform_D3D11(FName InShaderPlatform)
 
 static bool FilterShaderPlatform_Vulkan(FName InShaderPlatform)
 {
-	return InShaderPlatform == NAME_VULKAN_SM5;
+	return InShaderPlatform == NAME_VULKAN_SM5 || InShaderPlatform == NAME_VULKAN_SM6;
 }
 
 TSharedRef<IDetailCustomization> FWindowsTargetSettingsDetails::MakeInstance()
diff --git a/Engine/Source/Editor/MaterialEditor/Private/MaterialStats.cpp b/Engine/Source/Editor/MaterialEditor/Private/MaterialStats.cpp
--- a/Engine/Source/Editor/MaterialEditor/Private/MaterialStats.cpp
+++ b/Engine/Source/Editor/MaterialEditor/Private/MaterialStats.cpp
@@ -479,6 +479,7 @@ void FMaterialStats::BuildShaderPlatformDB()
 
 	// Vulkan
 	AddShaderPlatform(EPlatformCategoryType::Desktop, SP_VULKAN_SM5, TEXT("Vulkan SM5"), true, true, TEXT("Desktop, Vulkan, Shader Model 5"));
+	AddShaderPlatform(EPlatformCategoryType::Desktop, SP_VULKAN_SM6, TEXT("Vulkan SM6"), true, true, TEXT("Desktop, Vulkan, Shader Model 6"));
 
 	// Android
 	AddShaderPlatform(EPlatformCategoryType::Android, SP_OPENGL_ES3_1_ANDROID, TEXT("Android GLES 3.1"), true, true, TEXT("Android, OpenGLES 3.1"));
diff --git a/Engine/Source/Runtime/RHI/Private/RHIStrings.cpp b/Engine/Source/Runtime/RHI/Private/RHIStrings.cpp
--- a/Engine/Source/Runtime/RHI/Private/RHIStrings.cpp
+++ b/Engine/Source/Runtime/RHI/Private/RHIStrings.cpp
@@ -97,6 +97,7 @@ FString LexToString(EShaderPlatform Platform, bool bError)
 	case SP_VULKAN_ES3_1_ANDROID: return TEXT("VULKAN_ES3_1_ANDROID");
 	case SP_VULKAN_PCES3_1: return TEXT("VULKAN_PCES3_1");
 	case SP_VULKAN_SM5: return TEXT("VULKAN_SM5");
+	case SP_VULKAN_SM6: return TEXT("VULKAN_SM6");
 	case SP_VULKAN_SM5_ANDROID: return TEXT("VULKAN_SM5_ANDROID");
 	case SP_D3D_ES3_1_HOLOLENS: return TEXT("D3D_ES3_1_HOLOLENS");
 	case SP_CUSTOM_PLATFORM_FIRST: return TEXT("SP_CUSTOM_PLATFORM_FIRST");
@@ -202,6 +203,7 @@ FName ShaderPlatformToPlatformName(EShaderPlatform Platform)
 	case SP_OPENGL_PCES3_1:
 	case SP_VULKAN_PCES3_1:
 	case SP_VULKAN_SM5:
+	case SP_VULKAN_SM6:
 	case SP_D3D_ES3_1_HOLOLENS:
 		return NAME_PLATFORM_WINDOWS;
 	case SP_VULKAN_ES3_1_ANDROID:
diff --git a/Engine/Source/Runtime/RHI/Public/RHIShaderFormatDefinitions.inl b/Engine/Source/Runtime/RHI/Public/RHIShaderFormatDefinitions.inl
--- a/Engine/Source/Runtime/RHI/Public/RHIShaderFormatDefinitions.inl
+++ b/Engine/Source/Runtime/RHI/Public/RHIShaderFormatDefinitions.inl
@@ -27,6 +27,7 @@ static FName NAME_SF_METAL_MACES3_1(TEXT("SF_METAL_MACES3_1"));
 static FName NAME_VULKAN_ES3_1_ANDROID(TEXT("SF_VULKAN_ES31_ANDROID"));
 static FName NAME_VULKAN_ES3_1(TEXT("SF_VULKAN_ES31"));
 static FName NAME_VULKAN_SM5(TEXT("SF_VULKAN_SM5"));
+static FName NAME_VULKAN_SM6(TEXT("SF_VULKAN_SM6"));
 static FName NAME_VULKAN_SM5_ANDROID(TEXT("SF_VULKAN_SM5_ANDROID"));
 
 static FName NAME_D3D_ES3_1_HOLOLENS(TEXT("D3D_ES3_1_HOLOLENS"));
@@ -51,6 +52,7 @@ static EShaderPlatform ShaderFormatNameToShaderPlatform(FName ShaderFormat)
 	if (ShaderFormat == NAME_VULKAN_ES3_1_ANDROID)		return SP_VULKAN_ES3_1_ANDROID;
 	if (ShaderFormat == NAME_VULKAN_ES3_1)				return SP_VULKAN_PCES3_1;
 	if (ShaderFormat == NAME_VULKAN_SM5)				return SP_VULKAN_SM5;
+	if (ShaderFormat == NAME_VULKAN_SM6)				return SP_VULKAN_SM6;
 	if (ShaderFormat == NAME_VULKAN_SM5_ANDROID)		return SP_VULKAN_SM5_ANDROID;
 
 	if (ShaderFormat == NAME_D3D_ES3_1_HOLOLENS)		return SP_D3D_ES3_1_HOLOLENS;
diff --git a/Engine/Source/Runtime/RHI/Public/RHIShaderPlatform.h b/Engine/Source/Runtime/RHI/Public/RHIShaderPlatform.h
--- a/Engine/Source/Runtime/RHI/Public/RHIShaderPlatform.h
+++ b/Engine/Source/Runtime/RHI/Public/RHIShaderPlatform.h
@@ -47,6 +47,7 @@ enum EShaderPlatform : uint16
 	SP_VULKAN_SM5_ANDROID = SP_StaticPlatform_Last + 1,
 	SP_PCD3D_SM6,
 	SP_D3D_ES3_1_HOLOLENS,
+	SP_VULKAN_SM6,
 
 	SP_CUSTOM_PLATFORM_FIRST,
 	SP_CUSTOM_PLATFORM_LAST = (SP_CUSTOM_PLATFORM_FIRST + 100),
diff --git a/Engine/Source/Runtime/Renderer/Private/HairStrands/HairStrandsForwardRaster.cpp b/Engine/Source/Runtime/Renderer/Private/HairStrands/HairStrandsForwardRaster.cpp
--- a/Engine/Source/Runtime/Renderer/Private/HairStrands/HairStrandsForwardRaster.cpp
+++ b/Engine/Source/Runtime/Renderer/Private/HairStrands/HairStrandsForwardRaster.cpp
@@ -27,7 +27,8 @@ FMinHairRadiusAtDepth1 ComputeMinStrandRadiusAtDepth1(const FIntPoint& Resolutio
 
 inline bool IsHairStrandsForwardRasterSupported(EShaderPlatform In) 
 { 
-	return IsFeatureLevelSupported(In, ERHIFeatureLevel::SM6) && IsHairStrandsSupported(EHairStrandsShaderType::Strands, In); 
+	return IsFeatureLevelSupported(In, ERHIFeatureLevel::SM6) && IsHairStrandsSupported(EHairStrandsShaderType::Strands, In) && 
+		!IsVulkanPlatform(In); // :todo-jn: fix SPIR-V error during compilation
 }
 
 /////////////////////////////////////////////////////////////////////////////////////////
diff --git a/Engine/Source/Runtime/VulkanRHI/Private/VulkanGenericPlatform.cpp b/Engine/Source/Runtime/VulkanRHI/Private/VulkanGenericPlatform.cpp
--- a/Engine/Source/Runtime/VulkanRHI/Private/VulkanGenericPlatform.cpp
+++ b/Engine/Source/Runtime/VulkanRHI/Private/VulkanGenericPlatform.cpp
@@ -20,7 +20,7 @@ void FVulkanGenericPlatform::SetupFeatureLevels(TArrayView<EShaderPlatform> Shad
 	ShaderPlatformForFeatureLevel[ERHIFeatureLevel::ES3_1] = SP_VULKAN_PCES3_1;
 	ShaderPlatformForFeatureLevel[ERHIFeatureLevel::SM4_REMOVED] = SP_NumPlatforms;
 	ShaderPlatformForFeatureLevel[ERHIFeatureLevel::SM5] = SP_VULKAN_SM5;
-	ShaderPlatformForFeatureLevel[ERHIFeatureLevel::SM6] = SP_NumPlatforms;
+	ShaderPlatformForFeatureLevel[ERHIFeatureLevel::SM6] = SP_VULKAN_SM6;
 }
 
 bool FVulkanGenericPlatform::PSOBinaryCacheMatches(FVulkanDevice* Device, const TArray<uint8>& DeviceCache)
@@ -151,4 +151,4 @@ FString FVulkanGenericPlatform::GetVulkanProfileNameForFeatureLevel(ERHIFeatureL
 		ProfileName += TEXT("_RT");
 	}
 	return ProfileName;
-}
+}
\ No newline at end of file
