VulkanRHI: Fix Lumen tests for VULKAN_SM5

diff --git a/Engine/Config/VulkanPC/DataDrivenPlatformInfo.ini b/Engine/Config/VulkanPC/DataDrivenPlatformInfo.ini
--- a/Engine/Config/VulkanPC/DataDrivenPlatformInfo.ini
+++ b/Engine/Config/VulkanPC/DataDrivenPlatformInfo.ini
@@ -65,7 +65,7 @@ bRequiresVendorExtensionsForAtomics = false
 
 bSupportsNanite = false
 bSupportsSceneDataCompressedTransforms = true
-bSupportsLumenGI = false
+bSupportsLumenGI = true
 
 ; //#todo-vulkanddpi
 bSupportsSSDIndirect = false
@@ -235,7 +235,7 @@ bSupportsHairStrandGeometry=true
 bSupportsFFTBloom=true
 bSupportsDiaphragmDOF=true
 
-; shaderOutputLayer is optional in our baseline VK version
+; todo-jn: shaderOutputLayer is optional in our baseline VK version
 bSupportsVertexShaderLayer=false
 
 bSupportsAnisotropicMaterials=true
