Disables bSupportsRealTypes on vulkan du to cook failure colliding with r.Mobile.FloatPrecisionMode

diff --git a/Engine/Config/VulkanPC/DataDrivenPlatformInfo.ini b/Engine/Config/VulkanPC/DataDrivenPlatformInfo.ini
--- a/Engine/Config/VulkanPC/DataDrivenPlatformInfo.ini
+++ b/Engine/Config/VulkanPC/DataDrivenPlatformInfo.ini
@@ -134,7 +134,7 @@ FriendlyName=LOCTEXT("FriendlyShaderPlatformName_VulkanPC_SM5", "PC Vulkan SM5")
 
 BindlessSupport=Unsupported
 bSupportsNNEShaders=true
-bSupportsRealTypes=RuntimeDependent
+bSupportsRealTypes=Unsupported
 
 
 
@@ -251,7 +251,7 @@ FriendlyName=LOCTEXT("FriendlyShaderPlatformName_VulkanPC_SM6", "PC Vulkan SM6")
 
 BindlessSupport=AllShaderTypes
 bSupportsNNEShaders=true
-bSupportsRealTypes=RuntimeDependent
+bSupportsRealTypes=Unsupported
 
 
 [ShaderPlatform VULKAN_PCES3_1]
