VulkanRHI: document the ini values necessary to enable experimental Vulkan ray tracing.

diff --git a/Engine/Config/Linux/BaseLinuxEngine.ini b/Engine/Config/Linux/BaseLinuxEngine.ini
--- a/Engine/Config/Linux/BaseLinuxEngine.ini
+++ b/Engine/Config/Linux/BaseLinuxEngine.ini
@@ -10,6 +10,12 @@ DefaultExternalCredentialTypeStr=Steam
 [CrashReportClient]
 bAgreeToCrashUpload=false
 
+; To enable experimental Vulkan ray tracing, set both of these values to Enabled.
+; NOTE: Must also be enabled in VulkanPC/DataDrivenPlatformInfo.ini
+;[SF_VULKAN_SM6]
+;BindlessResources=Enabled
+;BindlessSamplers=Enabled
+
 [ConsoleVariables]
 ; necessary for LowMemoryWarning events to be fired
 memory.MemoryPressureCriticalThresholdMB=512
diff --git a/Engine/Config/VulkanPC/DataDrivenPlatformInfo.ini b/Engine/Config/VulkanPC/DataDrivenPlatformInfo.ini
--- a/Engine/Config/VulkanPC/DataDrivenPlatformInfo.ini
+++ b/Engine/Config/VulkanPC/DataDrivenPlatformInfo.ini
@@ -173,11 +173,17 @@ bSupports4ComponentUAVReadWrite = false
 bSupportsRenderTargetWriteMask = false
 
 bSupportsRayTracing = true
-bSupportsRayTracingShaders = false
 bSupportsInlineRayTracing = true
 bSupportsRayTracingIndirectInstanceData = true
 bSupportsByteBufferComputeShaders = true
 
+; To enable experimental Vulkan ray tracing, set the next four values to true.
+; NOTE: Must also enable Bindless in BaseWindowsEngine/BaseLinuxEngine.ini
+bSupportsRayTracingShaders = false
+bSupportsPathTracing = false
+bSupportsRayTracingCallableShaders = false
+bSupportsRayTracingProceduralPrimitive = false
+
 bSupportsGPUScene = true
 bSupportsPrimitiveShaders = false
 bSupportsUInt64ImageAtomics = true
diff --git a/Engine/Config/Windows/BaseWindowsEngine.ini b/Engine/Config/Windows/BaseWindowsEngine.ini
--- a/Engine/Config/Windows/BaseWindowsEngine.ini
+++ b/Engine/Config/Windows/BaseWindowsEngine.ini
@@ -23,6 +23,8 @@ BindlessResources=Disabled
 BindlessSamplers=Disabled
 
 [SF_VULKAN_SM6]
+; To enable experimental Vulkan ray tracing, set both of these values to Enabled.
+; NOTE: Must also be enabled in VulkanPC/DataDrivenPlatformInfo.ini
 BindlessResources=Disabled
 BindlessSamplers=Disabled
 
