VulkanRHI: Enable bindless and ray tracing pipelines by default in Vulkan SM6.

diff --git a/Engine/Config/VulkanPC/DataDrivenPlatformInfo.ini b/Engine/Config/VulkanPC/DataDrivenPlatformInfo.ini
--- a/Engine/Config/VulkanPC/DataDrivenPlatformInfo.ini
+++ b/Engine/Config/VulkanPC/DataDrivenPlatformInfo.ini
@@ -177,12 +177,11 @@ bSupportsInlineRayTracing = true
 bSupportsRayTracingIndirectInstanceData = true
 bSupportsByteBufferComputeShaders = true
 
-; To enable experimental Vulkan ray tracing, set the next four values to true.
-; NOTE: Must also enable Bindless in BaseWindowsEngine/BaseLinuxEngine.ini
-bSupportsRayTracingShaders = false
-bSupportsPathTracing = false
-bSupportsRayTracingCallableShaders = false
-bSupportsRayTracingProceduralPrimitive = false
+; Enable full Vulkan ray tracing
+bSupportsRayTracingShaders = true
+bSupportsPathTracing = true
+bSupportsRayTracingCallableShaders = true
+bSupportsRayTracingProceduralPrimitive = true
 
 bSupportsGPUScene = true
 bSupportsPrimitiveShaders = false
diff --git a/Engine/Config/Windows/BaseWindowsEngine.ini b/Engine/Config/Windows/BaseWindowsEngine.ini
--- a/Engine/Config/Windows/BaseWindowsEngine.ini
+++ b/Engine/Config/Windows/BaseWindowsEngine.ini
@@ -23,10 +23,9 @@ BindlessResources=Disabled
 BindlessSamplers=Disabled
 
 [SF_VULKAN_SM6]
-; To enable experimental Vulkan ray tracing, set both of these values to Enabled.
-; NOTE: Must also be enabled in VulkanPC/DataDrivenPlatformInfo.ini
-BindlessResources=Disabled
-BindlessSamplers=Disabled
+; Bindless is necessary for Vulkan ray tracing
+BindlessResources=Enabled
+BindlessSamplers=Enabled
 
 [/Script/WindowsRuntimeSettings.WindowsRuntimeSettings]
 MaxNumTlsSlots=2048
