VulkanRHI: Enable bSupportsVertexShaderLayer in DDPI

#rnx

diff --git a/Engine/Config/VulkanPC/DataDrivenPlatformInfo.ini b/Engine/Config/VulkanPC/DataDrivenPlatformInfo.ini
--- a/Engine/Config/VulkanPC/DataDrivenPlatformInfo.ini
+++ b/Engine/Config/VulkanPC/DataDrivenPlatformInfo.ini
@@ -118,8 +118,8 @@ bSupportsHairStrandGeometry=true
 bSupportsFFTBloom=true
 bSupportsDiaphragmDOF=true
 
-; shaderOutputLayer is optional in our baseline VK version
-bSupportsVertexShaderLayer=false
+; The support really is conditional, the runtime code needs to also test GRHISupportsArrayIndexFromAnyShader
+bSupportsVertexShaderLayer=true
 
 bSupportsAnisotropicMaterials=true
 
@@ -244,8 +244,7 @@ bSupportsHairStrandGeometry=true
 bSupportsFFTBloom=true
 bSupportsDiaphragmDOF=true
 
-; todo-jn: shaderOutputLayer is optional in our baseline VK version
-bSupportsVertexShaderLayer=false
+bSupportsVertexShaderLayer=true
 
 bSupportsAnisotropicMaterials=true
 
