VulkanRHI: Remove bRequiresDisableForwardLocalLights from Vulkan SM6 DDPI.

#rnx

diff --git a/Engine/Config/VulkanPC/DataDrivenPlatformInfo.ini b/Engine/Config/VulkanPC/DataDrivenPlatformInfo.ini
--- a/Engine/Config/VulkanPC/DataDrivenPlatformInfo.ini
+++ b/Engine/Config/VulkanPC/DataDrivenPlatformInfo.ini
@@ -227,9 +227,6 @@ bRequiresGeneratePrevTransformBuffer = false
 
 bRequiresRenderTargetDuringRaster = true
 
-; //#todo-vulkanddpi
-bRequiresDisableForwardLocalLights = true
-
 bCompileSignalProcessingPipeline = true
 bSupportsGen5TemporalAA=true
 
