VulkanRHI: Enable specialization constants for Vulkan in DDSPI.

#rnx
#Synced-CL 51414939

diff --git a/Engine/Config/Android/DataDrivenPlatformInfo.ini b/Engine/Config/Android/DataDrivenPlatformInfo.ini
--- a/Engine/Config/Android/DataDrivenPlatformInfo.ini
+++ b/Engine/Config/Android/DataDrivenPlatformInfo.ini
@@ -84,6 +84,7 @@ Language=Vulkan
 MaxFeatureLevel=ES3_1
 ShaderFormat=SF_VULKAN_ES31_ANDROID
 bIsMobile=true
+bSupportsSpecializationConstants=true
 bTargetsTiledGPU=true
 bSupportsMobileMultiView=true
 ; The support really is conditional, the runtime code needs to also test GRHISupportsArrayIndexFromAnyShader
diff --git a/Engine/Config/VulkanPC/DataDrivenPlatformInfo.ini b/Engine/Config/VulkanPC/DataDrivenPlatformInfo.ini
--- a/Engine/Config/VulkanPC/DataDrivenPlatformInfo.ini
+++ b/Engine/Config/VulkanPC/DataDrivenPlatformInfo.ini
@@ -12,6 +12,7 @@ GlobalIdentifier=97F8C85EF1B0424792D201C893358C8F
 Language=Vulkan
 MaxFeatureLevel=SM5
 ShaderFormat=SF_VULKAN_SM5
+bSupportsSpecializationConstants = true
 
 bIsMobile = false
 bIsMetalMRT = false
@@ -149,6 +150,7 @@ bIsConsole = false
 bIsAndroidOpenGLES = false
 bSupportsDxc = true
 bIsSPIRV=true
+bSupportsSpecializationConstants = true
 
 bSupportsDebugViewShaders=true
 bSupportsMobileMultiView = false
