VulkanRHI: Enable specialization constants in VULKAN_PCES3_1

#rnx

#Synced-CL 52548173

diff --git a/Engine/Config/VulkanPC/DataDrivenPlatformInfo.ini b/Engine/Config/VulkanPC/DataDrivenPlatformInfo.ini
--- a/Engine/Config/VulkanPC/DataDrivenPlatformInfo.ini
+++ b/Engine/Config/VulkanPC/DataDrivenPlatformInfo.ini
@@ -272,6 +272,7 @@ MaxFeatureLevel=ES3_1
 ShaderFormat=SF_VULKAN_ES31
 bIsPC=true
 bIsMobile=true
+bSupportsSpecializationConstants = true
 bSupportsDebugViewShaders=true
 ; Enables RDG pass merging(like on Android), helps with debugging
 bTargetsTiledGPU=true
